Use SetReplicatesByDefault in WheeledVehicleMovementComponent constructor

diff --git a/Engine/Plugins/Runtime/PhysXVehicles/Source/PhysXVehicles/Private/WheeledVehicleMovementComponent.cpp b/Engine/Plugins/Runtime/PhysXVehicles/Source/PhysXVehicles/Private/WheeledVehicleMovementComponent.cpp
--- a/Engine/Plugins/Runtime/PhysXVehicles/Source/PhysXVehicles/Private/WheeledVehicleMovementComponent.cpp
+++ b/Engine/Plugins/Runtime/PhysXVehicles/Source/PhysXVehicles/Private/WheeledVehicleMovementComponent.cpp
@@ -150,7 +150,7 @@ UWheeledVehicleMovementComponent::UWheeledVehicleMovementComponent(const FObject
 	MaxNormalizedTireLoadFiltered = PTireLoadFilterDef.mMaxFilteredNormalisedLoad;
 #endif // WITH_PHYSX
 
-	SetIsReplicated(true);
+	SetIsReplicatedByDefault(true);
 
 #if WITH_PHYSX_VEHICLES
 	AHUD::OnShowDebugInfo.AddUObject(this, &UWheeledVehicleMovementComponent::ShowDebugInfo);
@@ -314,7 +314,7 @@ void UWheeledVehicleMovementComponent::SetupVehicleShapes()
 
 				if(WheelBodySetup)
 				{
-					PxMeshScale MeshScale(U2PVector(UpdatedComponent->RelativeScale3D * MeshScaleV), PxQuat(physx::PxIdentity));
+					PxMeshScale MeshScale(U2PVector(UpdatedComponent->GetRelativeScale3D() * MeshScaleV), PxQuat(physx::PxIdentity));
 
 					if(WheelBodySetup->AggGeom.ConvexElems.Num() == 1)
 					{
@@ -718,7 +718,7 @@ FVector UWheeledVehicleMovementComponent::GetWheelRestingPosition( const FWheelS
 		USkinnedMeshComponent* Mesh = GetMesh();
 		if ( Mesh && Mesh->SkeletalMesh )
 		{
-			const FVector BonePosition = Mesh->SkeletalMesh->GetComposedRefPoseMatrix( WheelSetup.BoneName ).GetOrigin() * Mesh->RelativeScale3D;
+			const FVector BonePosition = Mesh->SkeletalMesh->GetComposedRefPoseMatrix( WheelSetup.BoneName ).GetOrigin() * Mesh->GetRelativeScale3D();
 			//BonePosition is local for the root BONE of the skeletal mesh - however, we are using the Root BODY which may have its own transform, so we need to return the position local to the root BODY
 			const FMatrix RootBodyMTX = Mesh->SkeletalMesh->GetComposedRefPoseMatrix(Mesh->GetBodyInstance()->BodySetup->BoneName);
 			const FVector LocalBonePosition = RootBodyMTX.InverseTransformPosition(BonePosition);
@@ -1943,7 +1943,7 @@ void UWheeledVehicleMovementComponent::CalculateAvoidanceVelocity(float DeltaTim
 		return;
 	}
 	
-	if (MyOwner->Role != ROLE_Authority)
+	if (MyOwner->GetLocalRole() != ROLE_Authority)
 	{	
 		return;
 	}
