#rnx

diff --git a/Engine/Platforms/HTML5/Plugins/Runtime/PhysXVehicles/PhysXVehicles_HTML5.uplugin b/Engine/Platforms/HTML5/Plugins/Runtime/PhysXVehicles/PhysXVehicles_HTML5.uplugin
--- a/Engine/Platforms/HTML5/Plugins/Runtime/PhysXVehicles/PhysXVehicles_HTML5.uplugin
+++ b/Engine/Platforms/HTML5/Plugins/Runtime/PhysXVehicles/PhysXVehicles_HTML5.uplugin
@@ -1,4 +1,4 @@
 {
 	"FileVersion" : 3,
-	"ParentPlugin": "PhysXVehicles"
+	"bIsPluginExtension": true
 }
