Disable PhysXVehicle plugin by default, Chaos Physics and ChaosVehiclePlugin are what's used in UE5

diff --git a/Engine/Plugins/Runtime/PhysXVehicles/PhysXVehicles.uplugin b/Engine/Plugins/Runtime/PhysXVehicles/PhysXVehicles.uplugin
--- a/Engine/Plugins/Runtime/PhysXVehicles/PhysXVehicles.uplugin
+++ b/Engine/Plugins/Runtime/PhysXVehicles/PhysXVehicles.uplugin
@@ -10,7 +10,7 @@
 	"DocsURL" : "",
 	"MarketplaceURL" : "",
 	"SupportURL" : "",
-	"EnabledByDefault" : true,
+	"EnabledByDefault" : false,
 	"CanContainContent" : false,
 	"IsBetaVersion" : false,
 	"Installed" : false,
